[MsQuic] Adding MsQuic plugin with MsQuicRuntime module

Adding the MsQuic runtime plugin with the MsQuicRuntime module to load the DLLs via the module interface once for every dependent plugin.

diff --git a/Engine/Plugins/Runtime/MsQuic/MsQuic.uplugin b/Engine/Plugins/Runtime/MsQuic/MsQuic.uplugin
new file mode 100644
--- /dev/null
+++ b/Engine/Plugins/Runtime/MsQuic/MsQuic.uplugin
@@ -0,0 +1,45 @@
+{
+	"FileVersion" : 3,
+	"Version" : 1,
+	"VersionName" : "1.0",
+	"FriendlyName" : "MsQuic Runtime Plugin",
+	"Description" : "Runtime plugin for the MsQuic library.",
+	"Category" : "Runtime",
+	"CreatedBy" : "Epic Games, Inc.",
+	"CreatedByURL" : "https://epicgames.com",
+	"DocsURL" : "",
+	"MarketplaceURL" : "",
+	"SupportURL" : "",
+	"EnabledByDefault" : true,
+	"CanContainContent" : false,
+	"IsBetaVersion" : false,
+	"Installed" : false,
+	"SupportedTargetPlatforms": [
+		"Win64",
+		"Linux",
+		"Mac"
+	],
+    "SupportedPrograms": [
+		"UnrealFrontend",
+		"UnrealMultiUserServer",
+		"UnrealMultiUserSlateServer",
+		"CrashReportClientEditor",
+		"CoopMultiUserServer"
+	],
+	"Modules" : [
+		{
+			"Name" : "MsQuicRuntime",
+			"Type" : "Runtime",
+			"LoadingPhase" : "PreDefault",
+			"PlatformAllowList": [ "Win64", "Linux", "Mac" ],
+			"ProgramAllowList": [
+				"UnrealFrontend",
+				"UnrealMultiUserServer",
+				"UnrealMultiUserSlateServer",
+				"CrashReportClientEditor",
+				"CoopMultiUserServer"
+			]
+		}
+	],
+	"Plugins": []
+}
diff --git a/Engine/Plugins/Runtime/MsQuic/Source/MsQuicRuntime.Build.cs b/Engine/Plugins/Runtime/MsQuic/Source/MsQuicRuntime.Build.cs
new file mode 100644
--- /dev/null
+++ b/Engine/Plugins/Runtime/MsQuic/Source/MsQuicRuntime.Build.cs
@@ -0,0 +1,18 @@
+// Copyright Epic Games, Inc. All Rights Reserved.
+
+namespace UnrealBuildTool.Rules
+{
+	public class MsQuicRuntime : ModuleRules
+	{
+		public MsQuicRuntime(ReadOnlyTargetRules Target) : base(Target)
+		{
+			PublicDependencyModuleNames.AddRange(
+				new string[]
+				{
+					"Core",
+                    "MsQuic"
+				}
+			);
+		}
+	}
+}
diff --git a/Engine/Plugins/Runtime/MsQuic/Source/Private/MsQuicRuntimeModule.cpp b/Engine/Plugins/Runtime/MsQuic/Source/Private/MsQuicRuntimeModule.cpp
new file mode 100644
--- /dev/null
+++ b/Engine/Plugins/Runtime/MsQuic/Source/Private/MsQuicRuntimeModule.cpp
@@ -0,0 +1,123 @@
+// Copyright Epic Games, Inc. All Rights Reserved.
+
+#include "MsQuicRuntimePrivate.h"
+#include "IMsQuicRuntimeModule.h"
+
+#include "CoreMinimal.h"
+#include "CoreTypes.h"
+#include "Misc/CoreMisc.h"
+#include "Misc/App.h"
+#include "Misc/Paths.h"
+#include "Modules/ModuleManager.h"
+
+DEFINE_LOG_CATEGORY(LogMsQuicRuntime);
+
+#define LOCTEXT_NAMESPACE "FMsQuicRuntimeModule"
+
+
+class FMsQuicRuntimeModule
+	: public IMsQuicRuntimeModule
+{
+
+public:
+
+	// IMsQuicRuntimeModule interface
+
+	bool InitRuntime() override
+    {
+		if (MsQuicLibraryHandle)
+		{
+			UE_LOG(LogMsQuicRuntime, Display,
+				TEXT("[MsQuicRuntimeModule] MsQuic DLL already loaded."));
+
+			return true;
+		}
+
+		if (!LoadMsQuicDll())
+		{
+			UE_LOG(LogMsQuicRuntime, Error,
+				TEXT("[MsQuicRuntimeModule] Could not load MsQuic DLL."));
+
+			return false;
+		}
+
+		return true;
+    }
+
+public:
+
+	// IModuleInterface interface
+
+    virtual void StartupModule() override
+    {
+    }
+
+    virtual void ShutdownModule() override
+    {
+		FreeMsQuicDll();
+    }
+
+private:
+
+	/**
+	 * Load the appropriate MsQuic DLL/So for this platform.
+	 */
+	bool LoadMsQuicDll()
+	{
+		const FString MsQuicBinariesDir = FPaths::Combine(
+			*FPaths::EngineDir(), *MSQUIC_BINARIES_PATH);
+
+		FString MsQuicLib = "";
+
+#if PLATFORM_WINDOWS
+
+		MsQuicLib = FPaths::Combine(*MsQuicBinariesDir,
+			TEXT("win64/msquic.dll"));
+
+#elif PLATFORM_LINUX
+
+		MsQuicLib = FPaths::Combine(*MsQuicBinariesDir,
+			TEXT("linux/libmsquic.so"));
+
+#elif PLATFORM_MAC
+
+		MsQuicLib = FPaths::Combine(*MsQuicBinariesDir,
+			TEXT("macos/libmsquic.dylib"));
+
+#endif
+
+		MsQuicLibraryHandle = (MsQuicLib.IsEmpty())
+			? nullptr : FPlatformProcess::GetDllHandle(*MsQuicLib);
+
+		return MsQuicLibraryHandle != nullptr;
+	}
+
+	/**
+	 * Free the MsQuic DLL/So if the LibraryHandle is valid.
+	 */
+	void FreeMsQuicDll()
+	{
+		if (MsQuicLibraryHandle)
+		{
+			FPlatformProcess::FreeDllHandle(MsQuicLibraryHandle);
+			MsQuicLibraryHandle = nullptr;
+		}
+	}
+
+private:
+
+	/** Holds the MsQuic library dll handle. */
+	void* MsQuicLibraryHandle = nullptr;
+
+	/** Defines the MsQuic version to be used. */
+	const FString MSQUIC_VERSION = "v220";
+
+	/** Defines the MsQuic binaries path. */
+	const FString MSQUIC_BINARIES_PATH = "Binaries/ThirdParty/MsQuic/" + MSQUIC_VERSION;
+
+};
+
+
+#undef LOCTEXT_NAMESPACE
+
+IMPLEMENT_MODULE(FMsQuicRuntimeModule, MsQuicRuntime);
diff --git a/Engine/Plugins/Runtime/MsQuic/Source/Private/MsQuicRuntimePrivate.h b/Engine/Plugins/Runtime/MsQuic/Source/Private/MsQuicRuntimePrivate.h
new file mode 100644
--- /dev/null
+++ b/Engine/Plugins/Runtime/MsQuic/Source/Private/MsQuicRuntimePrivate.h
@@ -0,0 +1,10 @@
+// Copyright Epic Games, Inc. All Rights Reserved.
+
+#pragma once
+
+#include "CoreMinimal.h"
+#include "Logging/LogMacros.h"
+
+
+/** Declares a log category for this module. */
+DECLARE_LOG_CATEGORY_EXTERN(LogMsQuicRuntime, Display, All);
diff --git a/Engine/Plugins/Runtime/MsQuic/Source/Public/IMsQuicRuntimeModule.h b/Engine/Plugins/Runtime/MsQuic/Source/Public/IMsQuicRuntimeModule.h
new file mode 100644
--- /dev/null
+++ b/Engine/Plugins/Runtime/MsQuic/Source/Public/IMsQuicRuntimeModule.h
@@ -0,0 +1,19 @@
+// Copyright Epic Games, Inc. All Rights Reserved.
+
+#pragma once
+
+#include "Modules/ModuleInterface.h"
+
+
+class IMsQuicRuntimeModule
+	: public IModuleInterface
+{
+
+public:
+
+	/**
+	 * Initiates the MsQuic runtime, loads the appropriate DLLs.
+	 */
+	virtual bool InitRuntime() = 0;
+
+};
