[MsQuic] Restructuring MsQuicRuntime module to have a static InitRuntime function

Restructuring the MsQuicRuntime module to have a static InitRuntime function that can be called from other plugins to load the MsQuic DLLs if they are not loaded yet.

diff --git a/Engine/Plugins/Runtime/MsQuic/Source/Private/MsQuicRuntimeModule.cpp b/Engine/Plugins/Runtime/MsQuic/Source/Private/MsQuicRuntimeModule.cpp
--- a/Engine/Plugins/Runtime/MsQuic/Source/Private/MsQuicRuntimeModule.cpp
+++ b/Engine/Plugins/Runtime/MsQuic/Source/Private/MsQuicRuntimeModule.cpp
@@ -1,7 +1,7 @@
 // Copyright Epic Games, Inc. All Rights Reserved.
 
+#include "MsQuicRuntimeModule.h"
 #include "MsQuicRuntimePrivate.h"
-#include "IMsQuicRuntimeModule.h"
 
 #include "CoreMinimal.h"
 #include "CoreTypes.h"
@@ -15,107 +15,76 @@ DEFINE_LOG_CATEGORY(LogMsQuicRuntime);
 #define LOCTEXT_NAMESPACE "FMsQuicRuntimeModule"
 
 
-class FMsQuicRuntimeModule
-	: public IMsQuicRuntimeModule
+bool FMsQuicRuntimeModule::InitRuntime()
 {
-
-public:
-
-	// IMsQuicRuntimeModule interface
-
-	bool InitRuntime() override
-    {
-		if (MsQuicLibraryHandle)
-		{
-			UE_LOG(LogMsQuicRuntime, Display,
-				TEXT("[MsQuicRuntimeModule] MsQuic DLL already loaded."));
-
-			return true;
-		}
-
-		if (!LoadMsQuicDll())
-		{
-			UE_LOG(LogMsQuicRuntime, Error,
-				TEXT("[MsQuicRuntimeModule] Could not load MsQuic DLL."));
-
-			return false;
-		}
+	if (MsQuicLibraryHandle)
+	{
+		UE_LOG(LogMsQuicRuntime, Display,
+			TEXT("[MsQuicRuntimeModule] MsQuic DLL already loaded."));
 
 		return true;
-    }
+	}
 
-public:
+	if (!LoadMsQuicDll())
+	{
+		UE_LOG(LogMsQuicRuntime, Error,
+			TEXT("[MsQuicRuntimeModule] Could not load MsQuic DLL."));
+
+		return false;
+	}
 
-	// IModuleInterface interface
+	return true;
+}
 
-    virtual void StartupModule() override
-    {
-    }
 
-    virtual void ShutdownModule() override
-    {
-		FreeMsQuicDll();
-    }
+void FMsQuicRuntimeModule::ShutdownModule()
+{
+	FreeMsQuicDll();
+}
 
-private:
 
-	/**
-	 * Load the appropriate MsQuic DLL/So for this platform.
-	 */
-	bool LoadMsQuicDll()
-	{
-		const FString MsQuicBinariesDir = FPaths::Combine(
-			*FPaths::EngineDir(), *MSQUIC_BINARIES_PATH);
+bool FMsQuicRuntimeModule::LoadMsQuicDll()
+{
+	const FString MsQuicBinariesDir = FPaths::Combine(
+		*FPaths::EngineDir(), *MSQUIC_BINARIES_PATH);
 
-		FString MsQuicLib = "";
+	FString MsQuicLib = "";
 
 #if PLATFORM_WINDOWS
 
-		MsQuicLib = FPaths::Combine(*MsQuicBinariesDir,
-			TEXT("win64/msquic.dll"));
+	MsQuicLib = FPaths::Combine(*MsQuicBinariesDir,
+		TEXT("win64/msquic.dll"));
 
 #elif PLATFORM_LINUX
 
-		MsQuicLib = FPaths::Combine(*MsQuicBinariesDir,
-			TEXT("linux/libmsquic.so"));
+	MsQuicLib = FPaths::Combine(*MsQuicBinariesDir,
+		TEXT("linux/libmsquic.so"));
 
 #elif PLATFORM_MAC
 
-		MsQuicLib = FPaths::Combine(*MsQuicBinariesDir,
-			TEXT("macos/libmsquic.dylib"));
+	MsQuicLib = FPaths::Combine(*MsQuicBinariesDir,
+		TEXT("macos/libmsquic.dylib"));
 
 #endif
 
-		MsQuicLibraryHandle = (MsQuicLib.IsEmpty())
-			? nullptr : FPlatformProcess::GetDllHandle(*MsQuicLib);
+	MsQuicLibraryHandle = (MsQuicLib.IsEmpty())
+		? nullptr : FPlatformProcess::GetDllHandle(*MsQuicLib);
+
+	return MsQuicLibraryHandle != nullptr;
+}
 
-		return MsQuicLibraryHandle != nullptr;
-	}
 
-	/**
-	 * Free the MsQuic DLL/So if the LibraryHandle is valid.
-	 */
-	void FreeMsQuicDll()
+/**
+ * Free the MsQuic DLL/So if the LibraryHandle is valid.
+ */
+void FMsQuicRuntimeModule::FreeMsQuicDll()
+{
+	if (MsQuicLibraryHandle)
 	{
-		if (MsQuicLibraryHandle)
-		{
-			FPlatformProcess::FreeDllHandle(MsQuicLibraryHandle);
-			MsQuicLibraryHandle = nullptr;
-		}
+		FPlatformProcess::FreeDllHandle(MsQuicLibraryHandle);
+		MsQuicLibraryHandle = nullptr;
 	}
-
-private:
-
-	/** Holds the MsQuic library dll handle. */
-	void* MsQuicLibraryHandle = nullptr;
-
-	/** Defines the MsQuic version to be used. */
-	const FString MSQUIC_VERSION = "v220";
-
-	/** Defines the MsQuic binaries path. */
-	const FString MSQUIC_BINARIES_PATH = "Binaries/ThirdParty/MsQuic/" + MSQUIC_VERSION;
-
-};
+}
 
 
 #undef LOCTEXT_NAMESPACE
diff --git a/Engine/Plugins/Runtime/MsQuic/Source/Public/IMsQuicRuntimeModule.h b/Engine/Plugins/Runtime/MsQuic/Source/Public/IMsQuicRuntimeModule.h
deleted file mode 100644
diff --git a/Engine/Plugins/Runtime/MsQuic/Source/Public/MsQuicRuntimeModule.h b/Engine/Plugins/Runtime/MsQuic/Source/Public/MsQuicRuntimeModule.h
new file mode 100644
--- /dev/null
+++ b/Engine/Plugins/Runtime/MsQuic/Source/Public/MsQuicRuntimeModule.h
@@ -0,0 +1,50 @@
+// Copyright Epic Games, Inc. All Rights Reserved.
+
+#pragma once
+
+#include "CoreMinimal.h"
+#include "Modules/ModuleInterface.h"
+
+
+class MSQUICRUNTIME_API FMsQuicRuntimeModule
+	: public IModuleInterface
+{
+
+public:
+
+	/**
+	 * Initializes the runtime by loading the MsQuic DLL if it's not loaded already.
+	 */
+	static bool InitRuntime();
+
+public:
+
+	// IModuleInterface interface override
+
+	virtual void ShutdownModule() override;
+
+private:
+
+	/**
+	 * Load the appropriate MsQuic DLL/So for this platform.
+	 */
+	static bool LoadMsQuicDll();
+
+	/**
+	 * Frees the DLL handle.
+	 */
+	static void FreeMsQuicDll();
+
+private:
+
+	/** Holds the MsQuic library dll handle. */
+	static inline void* MsQuicLibraryHandle = nullptr;
+
+	/** Defines the MsQuic version to be used. */
+	static inline const FString MSQUIC_VERSION = "v220";
+
+	/** Defines the MsQuic binaries path. */
+	static inline const FString MSQUIC_BINARIES_PATH
+		= "Binaries/ThirdParty/MsQuic/" + MSQUIC_VERSION;
+
+};
