[MsQuic] Fixing MsQuic linux build and runtime load of .so

Fixing the MsQuic linux build by adding a libmsquic.so that is a linker INPUT redirecting to libmsquic.so.2 and adapting the runtime dll init to load libmsquic.so.2 instead.

diff --git a/Engine/Plugins/Runtime/MsQuic/Source/Private/MsQuicRuntimeModule.cpp b/Engine/Plugins/Runtime/MsQuic/Source/Private/MsQuicRuntimeModule.cpp
--- a/Engine/Plugins/Runtime/MsQuic/Source/Private/MsQuicRuntimeModule.cpp
+++ b/Engine/Plugins/Runtime/MsQuic/Source/Private/MsQuicRuntimeModule.cpp
@@ -58,7 +58,7 @@ bool FMsQuicRuntimeModule::LoadMsQuicDll()
 #elif PLATFORM_LINUX
 
 	MsQuicLib = FPaths::Combine(*MsQuicBinariesDir,
-		TEXT("linux/libmsquic.so"));
+		TEXT("linux/libmsquic.so.2"));
 
 #elif PLATFORM_MAC
 
